UE5 - Fix for static analysis
#rnx

diff --git a/Engine/Plugins/Runtime/ChunkDownloader/Source/Private/ChunkDownloader.cpp b/Engine/Plugins/Runtime/ChunkDownloader/Source/Private/ChunkDownloader.cpp
--- a/Engine/Plugins/Runtime/ChunkDownloader/Source/Private/ChunkDownloader.cpp
+++ b/Engine/Plugins/Runtime/ChunkDownloader/Source/Private/ChunkDownloader.cpp
@@ -1757,6 +1757,7 @@ void FChunkDownloader::DownloadChunks(const TArray<int32>& ChunkIds, const FCall
 	// if there's no callback for some reason, avoid a bunch of boilerplate
 	if (Callback)
 	{
+		//+V773:SUPPRESS, class:FMultiCallback
 		// loop over chunks and issue individual callback
 		FMultiCallback* MultiCallback = new FMultiCallback(Callback);
 		for (const TSharedRef<FChunk>& Chunk : ChunksToDownload)
@@ -1764,14 +1765,17 @@ void FChunkDownloader::DownloadChunks(const TArray<int32>& ChunkIds, const FCall
 			DownloadChunkInternal(*Chunk, MultiCallback->AddPending(), Priority);
 		}
 		check(MultiCallback->GetNumPending() > 0);
-	} //-V773
+		//-V773
+	}
 	else
 	{
+		//+V773:SUPPRESS, class:FMultiCallback
 		// no need to manage callbacks
 		for (const TSharedRef<FChunk>& Chunk : ChunksToDownload)
 		{
 			DownloadChunkInternal(*Chunk, FCallback(), Priority);
 		}
+		//-V773
 	}
 
 	// resave manifest if needed
@@ -1842,6 +1846,7 @@ void FChunkDownloader::MountChunks(const TArray<int32>& ChunkIds, const FCallbac
 	// if there's no callback for some reason, avoid a bunch of boilerplate
 	if (Callback)
 	{
+		//+V773:SUPPRESS, class:FMultiCallback
 		// loop over chunks and issue individual callback
 		FMultiCallback* MultiCallback = new FMultiCallback(Callback);
 		for (const TSharedRef<FChunk>& Chunk : ChunksToMount)
@@ -1849,14 +1854,17 @@ void FChunkDownloader::MountChunks(const TArray<int32>& ChunkIds, const FCallbac
 			MountChunkInternal(*Chunk, MultiCallback->AddPending());
 		}
 		check(MultiCallback->GetNumPending() > 0);
-	} //-V773
+		//-V773
+	}
 	else
 	{
+		//+V773:SUPPRESS, class:FMultiCallback
 		// no need to manage callbacks
 		for (const TSharedRef<FChunk>& Chunk : ChunksToMount)
 		{
 			MountChunkInternal(*Chunk, FCallback());
 		}
+		//-V773
 	}
 
 	// resave manifest if needed
