UE-10688 - UE4 Static Analysis Win64 (PVS-Studio) - chunkdownloader.cpp - Visibility scope of the 'MultiCallback' pointer was exited without releasing the memory. A memory leak is possible

This is a test checkin to see if the extra lines supresses the issue. PVS-Studio is reporting the wrong line when running in horde but properly when running locally.

#rnx

diff --git a/Engine/Plugins/Runtime/ChunkDownloader/Source/Private/ChunkDownloader.cpp b/Engine/Plugins/Runtime/ChunkDownloader/Source/Private/ChunkDownloader.cpp
--- a/Engine/Plugins/Runtime/ChunkDownloader/Source/Private/ChunkDownloader.cpp
+++ b/Engine/Plugins/Runtime/ChunkDownloader/Source/Private/ChunkDownloader.cpp
@@ -1764,8 +1764,8 @@ void FChunkDownloader::DownloadChunks(const TArray<int32>& ChunkIds, const FCall
 			DownloadChunkInternal(*Chunk, MultiCallback->AddPending(), Priority);
 		}
 		check(MultiCallback->GetNumPending() > 0);
-	}  //-V773
-	else
+	} //-V773
+	else //-V773
 	{
 		// no need to manage callbacks
 		for (const TSharedRef<FChunk>& Chunk : ChunksToDownload)
@@ -1850,7 +1850,7 @@ void FChunkDownloader::MountChunks(const TArray<int32>& ChunkIds, const FCallbac
 		}
 		check(MultiCallback->GetNumPending() > 0);
 	} //-V773
-	else
+	else //-V773
 	{
 		// no need to manage callbacks
 		for (const TSharedRef<FChunk>& Chunk : ChunksToMount)
