UE-10688 - UE4 Static Analysis Win64 (PVS-Studio) - chunkdownloader.cpp - Visibility scope of the 'MultiCallback' pointer was exited without releasing the memory. A memory leak is possible

Removing prior change since it turns out the hack didn't help.

#rnx

diff --git a/Engine/Plugins/Runtime/ChunkDownloader/Source/Private/ChunkDownloader.cpp b/Engine/Plugins/Runtime/ChunkDownloader/Source/Private/ChunkDownloader.cpp
--- a/Engine/Plugins/Runtime/ChunkDownloader/Source/Private/ChunkDownloader.cpp
+++ b/Engine/Plugins/Runtime/ChunkDownloader/Source/Private/ChunkDownloader.cpp
@@ -1765,7 +1765,7 @@ void FChunkDownloader::DownloadChunks(const TArray<int32>& ChunkIds, const FCall
 		}
 		check(MultiCallback->GetNumPending() > 0);
 	} //-V773
-	else //-V773
+	else
 	{
 		// no need to manage callbacks
 		for (const TSharedRef<FChunk>& Chunk : ChunksToDownload)
@@ -1850,7 +1850,7 @@ void FChunkDownloader::MountChunks(const TArray<int32>& ChunkIds, const FCallbac
 		}
 		check(MultiCallback->GetNumPending() > 0);
 	} //-V773
-	else //-V773
+	else
 	{
 		// no need to manage callbacks
 		for (const TSharedRef<FChunk>& Chunk : ChunksToMount)
