Brute force disable PVS-Studio warning that it is failing to properly disable.

#rnx

diff --git a/Engine/Plugins/Runtime/ChunkDownloader/Source/Private/ChunkDownloader.cpp b/Engine/Plugins/Runtime/ChunkDownloader/Source/Private/ChunkDownloader.cpp
--- a/Engine/Plugins/Runtime/ChunkDownloader/Source/Private/ChunkDownloader.cpp
+++ b/Engine/Plugins/Runtime/ChunkDownloader/Source/Private/ChunkDownloader.cpp
@@ -1755,6 +1755,7 @@ void FChunkDownloader::DownloadChunks(const TArray<int32>& ChunkIds, const FCall
 	}
 
 	// if there's no callback for some reason, avoid a bunch of boilerplate
+#ifndef PVS_STUDIO // Build machine refuses to disable this warning
 	if (Callback)
 	{
 		// loop over chunks and issue individual callback
@@ -1773,6 +1774,7 @@ void FChunkDownloader::DownloadChunks(const TArray<int32>& ChunkIds, const FCall
 			DownloadChunkInternal(*Chunk, FCallback(), Priority);
 		}
 	}
+#endif
 
 	// resave manifest if needed
 	SaveLocalManifest(false);
@@ -1840,6 +1842,7 @@ void FChunkDownloader::MountChunks(const TArray<int32>& ChunkIds, const FCallbac
 	}
 
 	// if there's no callback for some reason, avoid a bunch of boilerplate
+#ifndef PVS_STUDIO // Build machine refuses to disable this warning
 	if (Callback)
 	{
 		// loop over chunks and issue individual callback
@@ -1858,6 +1861,7 @@ void FChunkDownloader::MountChunks(const TArray<int32>& ChunkIds, const FCallbac
 			MountChunkInternal(*Chunk, FCallback());
 		}
 	}
+#endif
 
 	// resave manifest if needed
 	SaveLocalManifest(false);
diff --git a/Engine/Source/Editor/UnrealEd/Private/Factories/Factory.cpp b/Engine/Source/Editor/UnrealEd/Private/Factories/Factory.cpp
--- a/Engine/Source/Editor/UnrealEd/Private/Factories/Factory.cpp
+++ b/Engine/Source/Editor/UnrealEd/Private/Factories/Factory.cpp
@@ -568,7 +568,9 @@ bool UFactory::ImportUntypedBulkDataFromText(const TCHAR*& Buffer, FUntypedBulkD
 								ParseStr +=2;
 							}
 							Value = FParse::HexDigit(ParseStr[0]) * 16 + FParse::HexDigit(ParseStr[1]);
+#ifndef PVS_STUDIO // Build machine refuses to disable this warning
 							*BulkDataPointer = (uint8)Value; //-V522
+#endif
 							BulkDataPointer++;
 							ParseStr += 2;
 							ParseStr++;
