Replace deprecated OnMountPak with MountPak.
#rnx

diff --git a/Engine/Plugins/Runtime/ChunkDownloader/Source/Private/ChunkDownloader.cpp b/Engine/Plugins/Runtime/ChunkDownloader/Source/Private/ChunkDownloader.cpp
--- a/Engine/Plugins/Runtime/ChunkDownloader/Source/Private/ChunkDownloader.cpp
+++ b/Engine/Plugins/Runtime/ChunkDownloader/Source/Private/ChunkDownloader.cpp
@@ -180,30 +180,24 @@ public:
 	void DoWork()
 	{
 		// try to mount the pak file
-		PRAGMA_DISABLE_DEPRECATION_WARNINGS
-		if (FCoreDelegates::OnMountPak.IsBound())
-		PRAGMA_ENABLE_DEPRECATION_WARNINGS
+		if (FCoreDelegates::MountPak.IsBound())
 		{
 			uint32 PakReadOrder = PakFiles.Num();
 			for (const TSharedRef<FPakFile>& PakFile : PakFiles)
 			{
 				FString FullPathOnDisk = (PakFile->bIsEmbedded ? EmbeddedFolder : CacheFolder) / PakFile->Entry.FileName;
-				PRAGMA_DISABLE_DEPRECATION_WARNINGS
-				bool bMountOk = FCoreDelegates::OnMountPak.Execute(FullPathOnDisk, PakReadOrder, nullptr);
-				PRAGMA_ENABLE_DEPRECATION_WARNINGS
+				IPakFile* MountedPak = FCoreDelegates::MountPak.Execute(FullPathOnDisk, PakReadOrder);
 
 #if !UE_BUILD_SHIPPING
-				if (!bMountOk)
+				if (!MountedPak)
 				{
 					// This can fail because of the sandbox system - which the pak system doesn't understand.
 					FString SandboxedPath = IFileManager::Get().ConvertToAbsolutePathForExternalAppForRead(*FullPathOnDisk);
-					PRAGMA_DISABLE_DEPRECATION_WARNINGS
-					bMountOk = FCoreDelegates::OnMountPak.Execute(SandboxedPath, PakReadOrder, nullptr);
-					PRAGMA_ENABLE_DEPRECATION_WARNINGS
+					MountedPak = FCoreDelegates::MountPak.Execute(SandboxedPath, PakReadOrder);
 				}
 #endif
 
-				if (bMountOk)
+				if (MountedPak)
 				{
 					// record that we successfully mounted this pak file
 					MountedPakFiles.Add(PakFile);
