PR #8240: [New][ChunkDownloader] Add GetNumDownloadRequests (Contributed by [REDACTED_PERSON])

diff --git a/Engine/Plugins/Runtime/ChunkDownloader/Source/Public/ChunkDownloader.h b/Engine/Plugins/Runtime/ChunkDownloader/Source/Public/ChunkDownloader.h
--- a/Engine/Plugins/Runtime/ChunkDownloader/Source/Public/ChunkDownloader.h
+++ b/Engine/Plugins/Runtime/ChunkDownloader/Source/Public/ChunkDownloader.h
@@ -140,6 +140,9 @@ public:
 	// called each time a download attempt finishes (success or failure). ONLY USE THIS IF YOU WANT TO PASSIVELY LISTEN. Downloads retry until successful.
 	TFunction<void(const FString& FileName, const FString& Url, uint64 SizeBytes, const FTimespan& DownloadTime, int32 HttpStatus)> OnDownloadAnalytics;
 
+	// get current number of download requests, so we know whether download is in progress. Downlading Requests will be removed from this array in it's FDownload::OnCompleted callback.
+	inline int32 GetNumDownloadRequests() const { return DownloadRequests.Num(); }
+
 protected:
 	friend class FChunkDownloaderModule;
 	friend class FChunkDownloaderPlatformWrapper;
