Cleaned up ChunkDownloader to not set a private pch file(it isn't force included) and to enable IWYU.

diff --git a/Engine/Plugins/Runtime/ChunkDownloader/Source/ChunkDownloader.Build.cs b/Engine/Plugins/Runtime/ChunkDownloader/Source/ChunkDownloader.Build.cs
--- a/Engine/Plugins/Runtime/ChunkDownloader/Source/ChunkDownloader.Build.cs
+++ b/Engine/Plugins/Runtime/ChunkDownloader/Source/ChunkDownloader.Build.cs
@@ -6,9 +6,6 @@ public class ChunkDownloader : ModuleRules
 {
 	public ChunkDownloader(ReadOnlyTargetRules Target) : base(Target)
 	{
-        PCHUsage = ModuleRules.PCHUsageMode.UseSharedPCHs;
-        PrivatePCHHeaderFile = "Private/ChunkDownloaderPCH.h";
-
         PrivateDependencyModuleNames.AddRange(
 			new string[] {
 				"Core",
diff --git a/Engine/Plugins/Runtime/ChunkDownloader/Source/Private/ChunkDownloader.cpp b/Engine/Plugins/Runtime/ChunkDownloader/Source/Private/ChunkDownloader.cpp
--- a/Engine/Plugins/Runtime/ChunkDownloader/Source/Private/ChunkDownloader.cpp
+++ b/Engine/Plugins/Runtime/ChunkDownloader/Source/Private/ChunkDownloader.cpp
@@ -1,7 +1,7 @@
 // Copyright Epic Games, Inc. All Rights Reserved.
 
-#include "ChunkDownloaderPCH.h"
 #include "ChunkDownloader.h"
+#include "ChunkDownloaderLog.h"
 #include "CoreMinimal.h"
 #include "Templates/UniquePtr.h"
 #include "Containers/Ticker.h"
diff --git a/Engine/Plugins/Runtime/ChunkDownloader/Source/Private/ChunkDownloaderLog.h b/Engine/Plugins/Runtime/ChunkDownloader/Source/Private/ChunkDownloaderLog.h
new file mode 100644
--- /dev/null
+++ b/Engine/Plugins/Runtime/ChunkDownloader/Source/Private/ChunkDownloaderLog.h
@@ -0,0 +1,13 @@
+// Copyright Epic Games, Inc. All Rights Reserved.
+
+#pragma once
+
+#include "CoreMinimal.h"
+#include "Misc/CoreMisc.h"
+#include "HAL/FileManagerGeneric.h"
+#include "Modules/ModuleManager.h"
+#include "Interfaces/IHttpResponse.h"
+#include "HttpModule.h"
+#include "PlatformHttp.h"
+
+DECLARE_LOG_CATEGORY_EXTERN(LogChunkDownloader, Log, All);
diff --git a/Engine/Plugins/Runtime/ChunkDownloader/Source/Private/ChunkDownloaderModule.cpp b/Engine/Plugins/Runtime/ChunkDownloader/Source/Private/ChunkDownloaderModule.cpp
--- a/Engine/Plugins/Runtime/ChunkDownloader/Source/Private/ChunkDownloaderModule.cpp
+++ b/Engine/Plugins/Runtime/ChunkDownloader/Source/Private/ChunkDownloaderModule.cpp
@@ -1,7 +1,7 @@
 // Copyright Epic Games, Inc. All Rights Reserved.
 
-#include "ChunkDownloaderPCH.h"
 #include "ChunkDownloader.h"
+#include "ChunkDownloaderLog.h"
 #include "Modules/ModuleManager.h"
 #include "GenericPlatform/GenericPlatformChunkInstall.h"
 
diff --git a/Engine/Plugins/Runtime/ChunkDownloader/Source/Private/ChunkDownloaderPCH.h b/Engine/Plugins/Runtime/ChunkDownloader/Source/Private/ChunkDownloaderPCH.h
deleted file mode 100644
diff --git a/Engine/Plugins/Runtime/ChunkDownloader/Source/Private/Download.cpp b/Engine/Plugins/Runtime/ChunkDownloader/Source/Private/Download.cpp
--- a/Engine/Plugins/Runtime/ChunkDownloader/Source/Private/Download.cpp
+++ b/Engine/Plugins/Runtime/ChunkDownloader/Source/Private/Download.cpp
@@ -1,7 +1,7 @@
 // Copyright Epic Games, Inc. All Rights Reserved.
 
-#include "ChunkDownloaderPCH.h"
 #include "Download.h"
+#include "ChunkDownloaderLog.h"
 #include "CoreMinimal.h"
 #include "Containers/Ticker.h"
 
diff --git a/Engine/Plugins/Runtime/ChunkDownloader/Source/Private/PlatformStreamDownload.cpp b/Engine/Plugins/Runtime/ChunkDownloader/Source/Private/PlatformStreamDownload.cpp
--- a/Engine/Plugins/Runtime/ChunkDownloader/Source/Private/PlatformStreamDownload.cpp
+++ b/Engine/Plugins/Runtime/ChunkDownloader/Source/Private/PlatformStreamDownload.cpp
@@ -1,7 +1,7 @@
 // Copyright Epic Games, Inc. All Rights Reserved.
 
-#include "ChunkDownloaderPCH.h"
 #include "PlatformStreamDownload.h"
+#include "ChunkDownloaderLog.h"
 
 //////////////////////////////////////////////////////////////////////////////////
 #if 0 && PLATFORM_ANDROID
diff --git a/Engine/Plugins/Runtime/ChunkDownloader/Source/Private/PlatformStreamDownload.h b/Engine/Plugins/Runtime/ChunkDownloader/Source/Private/PlatformStreamDownload.h
--- a/Engine/Plugins/Runtime/ChunkDownloader/Source/Private/PlatformStreamDownload.h
+++ b/Engine/Plugins/Runtime/ChunkDownloader/Source/Private/PlatformStreamDownload.h
@@ -2,6 +2,8 @@
 
 #pragma once
 
+#include "Templates/Function.h"
+
 typedef TFunction<void(int32 HttpStatus)> FDownloadComplete;
 typedef TFunction<void(int32 BytesReceived)> FDownloadProgress;
 typedef TFunction<void(void)> FDownloadCancel;
