Renamed private pch headers in ChunkDownloader and fixed -NoPCH -DisableUnity compile issues.

diff --git a/Engine/Plugins/Runtime/ChunkDownloader/Source/Private/PlatformStreamDownload.h b/Engine/Plugins/Runtime/ChunkDownloader/Source/Private/PlatformStreamDownload.h
--- a/Engine/Plugins/Runtime/ChunkDownloader/Source/Private/PlatformStreamDownload.h
+++ b/Engine/Plugins/Runtime/ChunkDownloader/Source/Private/PlatformStreamDownload.h
@@ -2,6 +2,7 @@
 
 #pragma once
 
+#include "HAL/Platform.h"
 #include "Templates/Function.h"
 
 typedef TFunction<void(int32 HttpStatus)> FDownloadComplete;
