Fix for potential out-of-bounds Read and Writes in ParseManifest()

Made the file position counters a consistent type
Changed FileSize type from int64 to const uint64.
Changed types for LineNum, NextLineStart, LineStart from int32 to uint64 so we no longer risk negative numbers when FileSize is large

diff --git a/Engine/Plugins/Runtime/ChunkDownloader/Source/Private/ChunkDownloader.cpp b/Engine/Plugins/Runtime/ChunkDownloader/Source/Private/ChunkDownloader.cpp
--- a/Engine/Plugins/Runtime/ChunkDownloader/Source/Private/ChunkDownloader.cpp
+++ b/Engine/Plugins/Runtime/ChunkDownloader/Source/Private/ChunkDownloader.cpp
@@ -261,7 +261,7 @@ static TArray<FPakFileEntry> ParseManifest(const FString& ManifestPath, TMap<FSt
 	IFileHandle* ManifestFile = IPlatformFile::GetPlatformPhysical().OpenRead(*ManifestPath);
 	if (ManifestFile != nullptr)
 	{
-		int64 FileSize = ManifestFile->Size();
+		const uint64 FileSize = ManifestFile->Size();
 		if (FileSize > 0)
 		{
 			UE_LOG(LogChunkDownloader, Log, TEXT("Found manifest at %s"), *ManifestPath);
@@ -280,15 +280,15 @@ static TArray<FPakFileEntry> ParseManifest(const FString& ManifestPath, TMap<FSt
 				char RelativeUrl[2048] = { 0 };
 
 				// line end
-				int LineNum = 0;
-				int NextLineStart = 0;
+				uint64 LineNum = 0;
+				uint64 NextLineStart = 0;
 				while (NextLineStart < FileSize)
 				{
-					int LineStart = NextLineStart;
+					uint64 LineStart = NextLineStart;
 					++LineNum;
 
 					// find the end of the line
-					int LineEnd = LineStart;
+					uint64 LineEnd = LineStart;
 					while (LineEnd < FileSize && FileBuffer[LineEnd] != '\n' && FileBuffer[LineEnd] != '\r')
 					{
 						++LineEnd;
